Added ShortName for GameplayBehaviorSmartObjectsModule
#rnx

diff --git a/Engine/Plugins/Runtime/GameplayBehaviorSmartObjects/Source/GameplayBehaviorSmartObjectsModule/GameplayBehaviorSmartObjectsModule.Build.cs b/Engine/Plugins/Runtime/GameplayBehaviorSmartObjects/Source/GameplayBehaviorSmartObjectsModule/GameplayBehaviorSmartObjectsModule.Build.cs
--- a/Engine/Plugins/Runtime/GameplayBehaviorSmartObjects/Source/GameplayBehaviorSmartObjectsModule/GameplayBehaviorSmartObjectsModule.Build.cs
+++ b/Engine/Plugins/Runtime/GameplayBehaviorSmartObjects/Source/GameplayBehaviorSmartObjectsModule/GameplayBehaviorSmartObjectsModule.Build.cs
@@ -6,6 +6,8 @@ public class GameplayBehaviorSmartObjectsModule : ModuleRules
 {
 	public GameplayBehaviorSmartObjectsModule(ReadOnlyTargetRules Target) : base(Target)
 	{
+		ShortName = "GameplayBhvSO";
+		
 		PCHUsage = ModuleRules.PCHUsageMode.UseExplicitOrSharedPCHs;
 		
 		PublicDependencyModuleNames.AddRange(
