UAITask_UseGameplayBehaviorSmartObject now validates that the claimed handle still points to a valid SmartObject before getting its location
#rnx

diff --git a/Engine/Plugins/Runtime/GameplayBehaviorSmartObjects/Source/GameplayBehaviorSmartObjectsModule/Private/AI/AITask_UseGameplayBehaviorSmartObject.cpp b/Engine/Plugins/Runtime/GameplayBehaviorSmartObjects/Source/GameplayBehaviorSmartObjectsModule/Private/AI/AITask_UseGameplayBehaviorSmartObject.cpp
--- a/Engine/Plugins/Runtime/GameplayBehaviorSmartObjects/Source/GameplayBehaviorSmartObjectsModule/Private/AI/AITask_UseGameplayBehaviorSmartObject.cpp
+++ b/Engine/Plugins/Runtime/GameplayBehaviorSmartObjects/Source/GameplayBehaviorSmartObjectsModule/Private/AI/AITask_UseGameplayBehaviorSmartObject.cpp
@@ -149,6 +149,13 @@ void UAITask_UseGameplayBehaviorSmartObject::Activate()
 		return;
 	}
 
+	// A valid claimed handle can point to an object that is no longer part of the simulation
+	if (!SmartObjectSubsystem->IsClaimedSmartObjectValid(ClaimedHandle))
+	{
+		UE_VLOG(OwnerController, LogSmartObject, Log, TEXT("Claim handle: %s refers to an object that is no longer available."), *LexToString(ClaimedHandle));
+		return;
+	}
+
 	const TOptional<FVector> GoalLocation = SmartObjectSubsystem->GetSlotLocation(ClaimedHandle);
 	if (!ensureMsgf(GoalLocation.IsSet(), TEXT("Unable to extract a valid slot location.")))
 	{
