Extended AITask_UseGameplayBehaviorSmartObject with more precise finish delegates allowing users to react to movement failures

diff --git a/Engine/Plugins/Runtime/GameplayBehaviorSmartObjects/Config/BaseGameplayBehaviorSmartObjects.ini b/Engine/Plugins/Runtime/GameplayBehaviorSmartObjects/Config/BaseGameplayBehaviorSmartObjects.ini
--- a/Engine/Plugins/Runtime/GameplayBehaviorSmartObjects/Config/BaseGameplayBehaviorSmartObjects.ini
+++ b/Engine/Plugins/Runtime/GameplayBehaviorSmartObjects/Config/BaseGameplayBehaviorSmartObjects.ini
@@ -1,7 +1,8 @@
-﻿[CoreRedirects]
+[CoreRedirects]
 +ClassRedirects=(OldName="/Script/SmartObjectsModule.SmartObjectGameplayBehaviorDefinition", NewName="/Script/GameplayBehaviorSmartObjectsModule.GameplayBehaviorSmartObjectBehaviorDefinition")
 +ClassRedirects=(OldName="/Script/SmartObjectsModule.AITask_UseSmartObject", NewName="/Script/GameplayBehaviorSmartObjectsModule.AITask_UseGameplayBehaviorSmartObject")
 +ClassRedirects=(OldName="/Script/SmartObjectsModule.BTTask_FindAndUseSmartObject", NewName="/Script/GameplayBehaviorSmartObjectsModule.BTTask_FindAndUseGameplayBehaviorSmartObject")
 +FunctionRedirects=(OldName="K2_UseSmartObject", NewName="/Script/GameplayBehaviorSmartObjectsModule.GameplayBehaviorSmartObjectsBlueprintFunctionLibrary.UseGameplayBehaviorSmartObject")
 +FunctionRedirects=(OldName="AITask_UseGameplayBehaviorSmartObject.UseSmartObject", NewName="AITask_UseGameplayBehaviorSmartObject.UseGameplayBehaviorSmartObject")
 +FunctionRedirects=(OldName="AITask_UseGameplayBehaviorSmartObject.UseClaimedSmartObject", NewName="AITask_UseGameplayBehaviorSmartObject.UseClaimedGameplayBehaviorSmartObject")
++PropertyRedirects=(OldName="/Script/GameplayBehaviorSmartObjects.AITask_UseGameplayBehaviorSmartObject.OnFinished",NewName="OnSucceeded")
diff --git a/Engine/Plugins/Runtime/GameplayBehaviorSmartObjects/Source/GameplayBehaviorSmartObjectsModule/Private/AI/AITask_UseGameplayBehaviorSmartObject.cpp b/Engine/Plugins/Runtime/GameplayBehaviorSmartObjects/Source/GameplayBehaviorSmartObjectsModule/Private/AI/AITask_UseGameplayBehaviorSmartObject.cpp
--- a/Engine/Plugins/Runtime/GameplayBehaviorSmartObjects/Source/GameplayBehaviorSmartObjectsModule/Private/AI/AITask_UseGameplayBehaviorSmartObject.cpp
+++ b/Engine/Plugins/Runtime/GameplayBehaviorSmartObjects/Source/GameplayBehaviorSmartObjectsModule/Private/AI/AITask_UseGameplayBehaviorSmartObject.cpp
@@ -167,6 +167,7 @@ void UAITask_UseGameplayBehaviorSmartObject::Activate()
 
 	FAIMoveRequest MoveReq(GoalLocation.GetValue());
 	MoveReq.SetUsePathfinding(true);
+	MoveReq.SetAllowPartialPath(false);
 
 	MoveToTask = UAITask::NewAITask<UAITask_MoveTo>(*OwnerController, *this, EAITaskPriority::High, TEXT("SmartObject"));
 	MoveToTask->SetUp(OwnerController, MoveReq);
@@ -177,7 +178,6 @@ void UAITask_UseGameplayBehaviorSmartObject::Activate()
 
 void UAITask_UseGameplayBehaviorSmartObject::OnGameplayTaskDeactivated(UGameplayTask& Task)
 {
-	UAITask::NewAITask<UAITask_MoveTo>(*OwnerController, *this, EAITaskPriority::High, TEXT("SmartObject"));
 	check(OwnerController);
 	check(OwnerController->GetPawn());
 
@@ -185,29 +185,38 @@ void UAITask_UseGameplayBehaviorSmartObject::OnGameplayTaskDeactivated(UGameplay
 	{
 		bool bBehaviorActive = false;
 
-		if (Task.IsFinished())
+		if (MoveToTask->IsFinished())
 		{
-			UWorld* World = OwnerController->GetWorld();
-			MoveToTask = nullptr;
-			USmartObjectSubsystem* SmartObjectSubsystem = USmartObjectSubsystem::GetCurrent(World);
-			if (ensure(SmartObjectSubsystem))
+			if (MoveToTask->WasMoveSuccessful())
 			{
-				const UGameplayBehaviorSmartObjectBehaviorDefinition* SmartObjectGameplayBehaviorDefinition = SmartObjectSubsystem->Use<UGameplayBehaviorSmartObjectBehaviorDefinition>(ClaimedHandle);
-				const UGameplayBehaviorConfig* GameplayBehaviorConfig = SmartObjectGameplayBehaviorDefinition != nullptr ? SmartObjectGameplayBehaviorDefinition->GameplayBehaviorConfig : nullptr;
-				GameplayBehavior = GameplayBehaviorConfig != nullptr ? GameplayBehaviorConfig->GetBehavior(*World) : nullptr;
-				if (GameplayBehavior != nullptr)
+				UWorld* World = OwnerController->GetWorld();
+				// Manually resetting MoveToTask ptr here in case the behavior coming after that is long running. 
+				// We don't need it anymore so if GC happens in the mean time it can clean it up.
+				MoveToTask = nullptr;
+				USmartObjectSubsystem* SmartObjectSubsystem = USmartObjectSubsystem::GetCurrent(World);
+				if (ensure(SmartObjectSubsystem))
 				{
-					const USmartObjectComponent* SmartObjectComponent = SmartObjectSubsystem->GetSmartObjectComponent(ClaimedHandle);
-					AActor& InteractorActor = *OwnerController->GetPawn();
-					AActor* InteracteeActor = SmartObjectComponent ? SmartObjectComponent->GetOwner() : nullptr;
-					bBehaviorActive = UGameplayBehaviorSubsystem::TriggerBehavior(*GameplayBehavior, InteractorActor, GameplayBehaviorConfig, InteracteeActor);
-					// Behavior can be successfully triggered AND ended synchronously. We are only interested to register callback when still running
-					if (bBehaviorActive)
+					const UGameplayBehaviorSmartObjectBehaviorDefinition* SmartObjectGameplayBehaviorDefinition = SmartObjectSubsystem->Use<UGameplayBehaviorSmartObjectBehaviorDefinition>(ClaimedHandle);
+					const UGameplayBehaviorConfig* GameplayBehaviorConfig = SmartObjectGameplayBehaviorDefinition != nullptr ? SmartObjectGameplayBehaviorDefinition->GameplayBehaviorConfig : nullptr;
+					GameplayBehavior = GameplayBehaviorConfig != nullptr ? GameplayBehaviorConfig->GetBehavior(*World) : nullptr;
+					if (GameplayBehavior != nullptr)
 					{
-						OnBehaviorFinishedNotifyHandle = GameplayBehavior->GetOnBehaviorFinishedDelegate().AddUObject(this, &UAITask_UseGameplayBehaviorSmartObject::OnSmartObjectBehaviorFinished);
+						const USmartObjectComponent* SmartObjectComponent = SmartObjectSubsystem->GetSmartObjectComponent(ClaimedHandle);
+						AActor& InteractorActor = *OwnerController->GetPawn();
+						AActor* InteracteeActor = SmartObjectComponent ? SmartObjectComponent->GetOwner() : nullptr;
+						bBehaviorActive = UGameplayBehaviorSubsystem::TriggerBehavior(*GameplayBehavior, InteractorActor, GameplayBehaviorConfig, InteracteeActor);
+						// Behavior can be successfully triggered AND ended synchronously. We are only interested to register callback when still running
+						if (bBehaviorActive)
+						{
+							OnBehaviorFinishedNotifyHandle = GameplayBehavior->GetOnBehaviorFinishedDelegate().AddUObject(this, &UAITask_UseGameplayBehaviorSmartObject::OnSmartObjectBehaviorFinished);
+						}
 					}
 				}
 			}
+			else
+			{
+				OnMoveToFailed.Broadcast();
+			}
 		}
 
 		if (bBehaviorActive == false)
@@ -234,7 +243,14 @@ void UAITask_UseGameplayBehaviorSmartObject::OnDestroy(const bool bInOwnerFinish
 
 	if (TaskState != EGameplayTaskState::Finished)
 	{
-		OnFinished.Broadcast();
+		if (GameplayBehavior)
+		{
+			OnSucceeded.Broadcast();
+		}
+		else
+		{
+			OnFailed.Broadcast();
+		}
 
 		Super::OnDestroy(bInOwnerFinished);
 	}
diff --git a/Engine/Plugins/Runtime/GameplayBehaviorSmartObjects/Source/GameplayBehaviorSmartObjectsModule/Public/AI/AITask_UseGameplayBehaviorSmartObject.h b/Engine/Plugins/Runtime/GameplayBehaviorSmartObjects/Source/GameplayBehaviorSmartObjectsModule/Public/AI/AITask_UseGameplayBehaviorSmartObject.h
--- a/Engine/Plugins/Runtime/GameplayBehaviorSmartObjects/Source/GameplayBehaviorSmartObjectsModule/Public/AI/AITask_UseGameplayBehaviorSmartObject.h
+++ b/Engine/Plugins/Runtime/GameplayBehaviorSmartObjects/Source/GameplayBehaviorSmartObjectsModule/Public/AI/AITask_UseGameplayBehaviorSmartObject.h
@@ -48,7 +48,13 @@ protected:
 
 protected:
 	UPROPERTY(BlueprintAssignable)
-	FGenericGameplayTaskDelegate OnFinished;
+	FGenericGameplayTaskDelegate OnSucceeded;
+
+	UPROPERTY(BlueprintAssignable)
+	FGenericGameplayTaskDelegate OnFailed;
+
+	UPROPERTY(BlueprintAssignable)
+	FGenericGameplayTaskDelegate OnMoveToFailed;
 
 	UPROPERTY()
 	TObjectPtr<UAITask_MoveTo> MoveToTask;
