Modified UBTTask_FindAndUseGameplayBehaviorSmartObject to support finding target smart object with an EQS query.

diff --git a/Engine/Plugins/Runtime/GameplayBehaviorSmartObjects/Source/GameplayBehaviorSmartObjectsModule/Private/AI/BTTask_FindAndUseGameplayBehaviorSmartObject.cpp b/Engine/Plugins/Runtime/GameplayBehaviorSmartObjects/Source/GameplayBehaviorSmartObjectsModule/Private/AI/BTTask_FindAndUseGameplayBehaviorSmartObject.cpp
--- a/Engine/Plugins/Runtime/GameplayBehaviorSmartObjects/Source/GameplayBehaviorSmartObjectsModule/Private/AI/BTTask_FindAndUseGameplayBehaviorSmartObject.cpp
+++ b/Engine/Plugins/Runtime/GameplayBehaviorSmartObjects/Source/GameplayBehaviorSmartObjectsModule/Private/AI/BTTask_FindAndUseGameplayBehaviorSmartObject.cpp
@@ -7,13 +7,16 @@
 #include "GameplayBehaviorSmartObjectBehaviorDefinition.h"
 #include "SmartObjectSubsystem.h"
 #include "VisualLogger/VisualLogger.h"
+#include "EnvironmentQuery/EnvQueryManager.h"
+#include "EnvQueryItemType_SmartObject.h"
 
 
-UBTTask_FindAndUseGameplayBehaviorSmartObject::UBTTask_FindAndUseGameplayBehaviorSmartObject(const FObjectInitializer& ObjectInitializer)
-	: Super(ObjectInitializer)
+UBTTask_FindAndUseGameplayBehaviorSmartObject::UBTTask_FindAndUseGameplayBehaviorSmartObject()
 {
-	//TagMatchingPolicy = ESmartObjectGameplayTagMatching::Any;
 	Radius = 500.f;
+	EQSQueryFinishedDelegate = FQueryFinishedSignature::CreateUObject(this, &UBTTask_FindAndUseGameplayBehaviorSmartObject::OnQueryFinished);
+	EQSRequest.RunMode = EEnvQueryRunMode::AllMatching;
+	bNotifyTaskFinished = true;
 }
 
 EBTNodeResult::Type UBTTask_FindAndUseGameplayBehaviorSmartObject::ExecuteTask(UBehaviorTreeComponent& OwnerComp, uint8* NodeMemory)
@@ -21,9 +24,9 @@ EBTNodeResult::Type UBTTask_FindAndUseGameplayBehaviorSmartObject::ExecuteTask(U
 	EBTNodeResult::Type NodeResult = EBTNodeResult::Failed;
 
 	UWorld* World = GetWorld();
-	USmartObjectSubsystem* Subsystem = USmartObjectSubsystem::GetCurrent(World);
+	USmartObjectSubsystem* SmartObjectSubsystem = USmartObjectSubsystem::GetCurrent(World);
 	AAIController* MyController = OwnerComp.GetAIOwner();
-	if (Subsystem == nullptr || MyController == nullptr
+	if (SmartObjectSubsystem == nullptr || MyController == nullptr
 		|| MyController->GetPawn() == nullptr)
 	{
 		return EBTNodeResult::Failed;
@@ -31,48 +34,61 @@ EBTNodeResult::Type UBTTask_FindAndUseGameplayBehaviorSmartObject::ExecuteTask(U
 
 	FBTUseSOTaskMemory* MyMemory = reinterpret_cast<FBTUseSOTaskMemory*>(NodeMemory);
 	MyMemory->TaskInstance.Reset();
+	MyMemory->EQSRequestID = INDEX_NONE;
 
 	AActor& Avatar = *MyController->GetPawn();
-	const FVector UserLocation = Avatar.GetActorLocation();
-
-	// Create filter
-	FSmartObjectRequestFilter Filter;
-	Filter.ActivityRequirements = ActivityRequirements;
-	Filter.BehaviorDefinitionClass = UGameplayBehaviorSmartObjectBehaviorDefinition::StaticClass();
-	const IGameplayTagAssetInterface* TagsSource = Cast<const IGameplayTagAssetInterface>(&Avatar);
-	if (TagsSource != nullptr)
+
+	if (EQSRequest.IsValid() && (EQSRequest.EQSQueryBlackboardKey.IsSet() || EQSRequest.QueryTemplate))
 	{
-		TagsSource->GetOwnedGameplayTags(Filter.UserTags);
+		const UBlackboardComponent* BlackboardComponent = OwnerComp.GetBlackboardComponent();
+		MyMemory->EQSRequestID = EQSRequest.Execute(Avatar, BlackboardComponent, EQSQueryFinishedDelegate);
+
+		if (MyMemory->EQSRequestID != INDEX_NONE)
+		{
+			NodeResult = EBTNodeResult::InProgress;
+		}
 	}
+	else 
+	{
+		const FVector UserLocation = Avatar.GetActorLocation();
+
+		// Create filter
+		FSmartObjectRequestFilter Filter;
+		Filter.ActivityRequirements = ActivityRequirements;
+		Filter.BehaviorDefinitionClass = UGameplayBehaviorSmartObjectBehaviorDefinition::StaticClass();
+		const IGameplayTagAssetInterface* TagsSource = Cast<const IGameplayTagAssetInterface>(&Avatar);
+		if (TagsSource != nullptr)
+		{
+			TagsSource->GetOwnedGameplayTags(Filter.UserTags);
+		}
 
-	// Create request
-	FSmartObjectRequest Request(FBox(UserLocation, UserLocation).ExpandBy(FVector(Radius), FVector(Radius)), Filter);
-	TArray<FSmartObjectRequestResult> Results; 
+		// Create request
+		FSmartObjectRequest Request(FBox(UserLocation, UserLocation).ExpandBy(FVector(Radius), FVector(Radius)), Filter);
+		TArray<FSmartObjectRequestResult> Results; 
 	
-	if (Subsystem->FindSmartObjects(Request, Results))
-	{
-		for (const FSmartObjectRequestResult& Result : Results)
+		if (SmartObjectSubsystem->FindSmartObjects(Request, Results))
 		{
-			FSmartObjectClaimHandle ClaimHandle = Subsystem->Claim(Result);
-			if (ClaimHandle.IsValid())
+			for (const FSmartObjectRequestResult& Result : Results)
 			{
-				UAITask_UseGameplayBehaviorSmartObject* UseSOTask = NewBTAITask<UAITask_UseGameplayBehaviorSmartObject>(OwnerComp);
-				UseSOTask->SetClaimHandle(ClaimHandle);
-				UseSOTask->ReadyForActivation();
-
-				NodeResult = EBTNodeResult::InProgress;
-				UE_VLOG_UELOG(MyController, LogSmartObject, Verbose, TEXT("%s claimed smart object: %s"), *GetNodeName(), *LexToString(ClaimHandle));
-				break;
+				FSmartObjectClaimHandle ClaimHandle = SmartObjectSubsystem->Claim(Result);
+				if (ClaimHandle.IsValid())
+				{
+					UseClaimedSmartObject(OwnerComp, ClaimHandle, *MyMemory);
+
+					NodeResult = EBTNodeResult::InProgress;
+					UE_VLOG_UELOG(MyController, LogSmartObject, Verbose, TEXT("%s claimed smart object: %s"), *GetNodeName(), *LexToString(ClaimHandle));
+					break;
+				}
 			}
-		}
 
-		UE_CVLOG_UELOG(NodeResult == EBTNodeResult::Failed, MyController, LogSmartObject, Warning, TEXT("%s failed to claim smart object"), *GetNodeName());
-	}
-	else
-	{
-		UE_VLOG_UELOG(MyController, LogSmartObject
-			, Verbose, TEXT("%s failed to find smart objects for request: %s")
-			, *GetNodeName(), *Avatar.GetName());
+			UE_CVLOG_UELOG(NodeResult == EBTNodeResult::Failed, MyController, LogSmartObject, Warning, TEXT("%s failed to claim smart object"), *GetNodeName());
+		}
+		else
+		{
+			UE_VLOG_UELOG(MyController, LogSmartObject
+				, Verbose, TEXT("%s failed to find smart objects for request: %s")
+				, *GetNodeName(), *Avatar.GetName());
+		}
 	}
 
 	return NodeResult;
@@ -80,12 +96,38 @@ EBTNodeResult::Type UBTTask_FindAndUseGameplayBehaviorSmartObject::ExecuteTask(U
 
 EBTNodeResult::Type UBTTask_FindAndUseGameplayBehaviorSmartObject::AbortTask(UBehaviorTreeComponent& OwnerComp, uint8* NodeMemory)
 {
+	check(NodeMemory);
+
+	FBTUseSOTaskMemory* MyMemory = reinterpret_cast<FBTUseSOTaskMemory*>(NodeMemory);
+	if (UAITask_UseGameplayBehaviorSmartObject* UseSOTask = MyMemory->TaskInstance.Get())
+	{
+		UseSOTask->ExternalCancel();
+		MyMemory->TaskInstance.Reset();
+	}
+
+	if (MyMemory->EQSRequestID != INDEX_NONE)
+	{
+		if (UWorld* World = OwnerComp.GetWorld())
+		{
+			if (UEnvQueryManager* EnvQueryManager = UEnvQueryManager::GetCurrent(World))
+			{
+				EnvQueryManager->AbortQuery(MyMemory->EQSRequestID);
+			}
+		}
+		MyMemory->EQSRequestID = INDEX_NONE;
+	}
+
 	return EBTNodeResult::Aborted;
 }
 
 void UBTTask_FindAndUseGameplayBehaviorSmartObject::OnTaskFinished(UBehaviorTreeComponent& OwnerComp, uint8* NodeMemory, EBTNodeResult::Type TaskResult)
 {
-
+	FBTUseSOTaskMemory* MyMemory = reinterpret_cast<FBTUseSOTaskMemory*>(NodeMemory);
+	if (UAITask_UseGameplayBehaviorSmartObject* UseSOTask = MyMemory->TaskInstance.Get())
+	{
+		check(UseSOTask->IsFinished());
+	}
+	check(MyMemory->EQSRequestID == INDEX_NONE);
 }
 
 FString UBTTask_FindAndUseGameplayBehaviorSmartObject::GetStaticDescription() const
@@ -100,4 +142,99 @@ FString UBTTask_FindAndUseGameplayBehaviorSmartObject::GetStaticDescription() co
 	return Result.Len() > 0
 		? Result
 		: Super::GetStaticDescription();
-}
\ No newline at end of file
+}
+
+void UBTTask_FindAndUseGameplayBehaviorSmartObject::InitializeFromAsset(UBehaviorTree& Asset)
+{
+	Super::InitializeFromAsset(Asset);
+	EQSRequest.InitForOwnerAndBlackboard(*this, GetBlackboardAsset());
+}
+
+void UBTTask_FindAndUseGameplayBehaviorSmartObject::OnQueryFinished(TSharedPtr<FEnvQueryResult> Result)
+{
+	if (!Result)
+	{
+		return;
+	}
+
+	AActor* MyOwner = Cast<AActor>(Result->Owner.Get());
+	if (APawn* PawnOwner = Cast<APawn>(MyOwner))
+	{
+		MyOwner = PawnOwner->GetController();
+	}
+
+	UBehaviorTreeComponent* BTComponent = MyOwner ? MyOwner->FindComponentByClass<UBehaviorTreeComponent>() : NULL;
+	if (!BTComponent)
+	{
+		UE_LOG(LogBehaviorTree, Warning, TEXT("%s [%s]: Unable to find behavior tree to notify about finished query!")
+			, ANSI_TO_TCHAR(__FUNCTION__), *GetNameSafe(MyOwner));
+		return;
+	}
+
+	const FEnvQueryResult& QueryResult = *Result.Get();
+
+	uint8* RawMemory = BTComponent->GetNodeMemory(this, BTComponent->FindInstanceContainingNode(this));
+	check(RawMemory);
+	FBTUseSOTaskMemory* MyMemory = reinterpret_cast<FBTUseSOTaskMemory*>(RawMemory);
+	if (MyMemory->EQSRequestID != QueryResult.QueryID)
+	{
+		UE_VLOG_UELOG(BTComponent, LogBehaviorTree, Log, TEXT("%s [%s] ignoring EQS result due to QueryID mismatch.")
+			, ANSI_TO_TCHAR(__FUNCTION__), *GetNameSafe(MyOwner));
+
+		check(MyMemory->EQSRequestID != INDEX_NONE)
+
+		return;
+	}
+	else if (QueryResult.IsAborted())
+	{
+		UE_VLOG_UELOG(BTComponent, LogBehaviorTree, Log, TEXT("%s [%s] observed EQS query finished as Aborted. Aborting the BT node as well.")
+			, ANSI_TO_TCHAR(__FUNCTION__), *GetNameSafe(MyOwner));
+
+		FinishLatentTask(*BTComponent, EBTNodeResult::Aborted);
+		return;
+	}
+
+	bool bSmartObjectClaimed = false;
+
+	if (QueryResult.IsSuccessful() && (QueryResult.Items.Num() >= 1))
+	{
+		if (QueryResult.ItemType->IsChildOf(UEnvQueryItemType_SmartObject::StaticClass()) == false)
+		{
+			UE_VLOG_UELOG(BTComponent, LogSmartObject, Error, TEXT("%s used EQS query that did not generate EnvQueryItemType_SmartObject items"), *GetNodeName());
+		}
+		else if (USmartObjectSubsystem* SmartObjectSubsystem = USmartObjectSubsystem::GetCurrent(MyOwner->GetWorld()))
+		{
+			// we could use QueryResult.GetItemAsTypeChecked, but the below implementation is more efficient
+			for (int i = 0; i < QueryResult.Items.Num(); ++i)
+			{
+				const FSmartObjectSlotEQSItem& Item = UEnvQueryItemType_SmartObject::GetValue(QueryResult.GetItemRawMemory(i));
+				const FSmartObjectClaimHandle ClaimHandle = SmartObjectSubsystem->Claim(Item.SmartObjectHandle, Item.SlotHandle);
+				if (ClaimHandle.IsValid())
+				{
+					UseClaimedSmartObject(*BTComponent, ClaimHandle, *MyMemory);
+					bSmartObjectClaimed = true;
+
+					UE_VLOG_UELOG(BTComponent, LogSmartObject, Verbose, TEXT("%s claimed EQS-found smart object: %s"), *GetNodeName(), *LexToString(ClaimHandle));
+					break;
+				}
+			}
+		}
+	}
+
+	MyMemory->EQSRequestID = INDEX_NONE;
+	
+	if (bSmartObjectClaimed == false)
+	{
+		FinishLatentTask(*BTComponent, EBTNodeResult::Failed);
+	}
+}
+
+void UBTTask_FindAndUseGameplayBehaviorSmartObject::UseClaimedSmartObject(UBehaviorTreeComponent& OwnerComp, FSmartObjectClaimHandle ClaimHandle, FBTUseSOTaskMemory& MyMemory)
+{
+	checkSlow(ClaimHandle.IsValid());
+	UAITask_UseGameplayBehaviorSmartObject* UseSOTask = NewBTAITask<UAITask_UseGameplayBehaviorSmartObject>(OwnerComp);
+	UseSOTask->SetClaimHandle(ClaimHandle);
+	UseSOTask->ReadyForActivation();
+
+	MyMemory.TaskInstance = UseSOTask;
+}
diff --git a/Engine/Plugins/Runtime/GameplayBehaviorSmartObjects/Source/GameplayBehaviorSmartObjectsModule/Public/AI/BTTask_FindAndUseGameplayBehaviorSmartObject.h b/Engine/Plugins/Runtime/GameplayBehaviorSmartObjects/Source/GameplayBehaviorSmartObjectsModule/Public/AI/BTTask_FindAndUseGameplayBehaviorSmartObject.h
--- a/Engine/Plugins/Runtime/GameplayBehaviorSmartObjects/Source/GameplayBehaviorSmartObjectsModule/Public/AI/BTTask_FindAndUseGameplayBehaviorSmartObject.h
+++ b/Engine/Plugins/Runtime/GameplayBehaviorSmartObjects/Source/GameplayBehaviorSmartObjectsModule/Public/AI/BTTask_FindAndUseGameplayBehaviorSmartObject.h
@@ -4,6 +4,7 @@
 
 #include "BehaviorTree/BTTaskNode.h"
 #include "AI/AITask_UseGameplayBehaviorSmartObject.h"
+#include "EnvironmentQuery/EnvQueryTypes.h"
 #include "BTTask_FindAndUseGameplayBehaviorSmartObject.generated.h"
 
 
@@ -12,6 +13,7 @@ class AITask_UseSmartObject;
 struct FBTUseSOTaskMemory
 {
 	TWeakObjectPtr<UAITask_UseGameplayBehaviorSmartObject> TaskInstance;
+	int32 EQSRequestID;
 };
 
 /**
@@ -22,8 +24,9 @@ class GAMEPLAYBEHAVIORSMARTOBJECTSMODULE_API UBTTask_FindAndUseGameplayBehaviorS
 {
 	GENERATED_BODY()
 public:
-	UBTTask_FindAndUseGameplayBehaviorSmartObject(const FObjectInitializer& ObjectInitializer = FObjectInitializer::Get());
+	UBTTask_FindAndUseGameplayBehaviorSmartObject();
 
+protected:
 	virtual EBTNodeResult::Type ExecuteTask(UBehaviorTreeComponent& OwnerComp, uint8* NodeMemory) override;
 	virtual EBTNodeResult::Type AbortTask(UBehaviorTreeComponent& OwnerComp, uint8* NodeMemory) override;
 	virtual void OnTaskFinished(UBehaviorTreeComponent& OwnerComp, uint8* NodeMemory, EBTNodeResult::Type TaskResult) override;
@@ -31,6 +34,12 @@ public:
 
 	virtual FString GetStaticDescription() const override;
 
+	virtual void InitializeFromAsset(UBehaviorTree& Asset) override;
+
+	void OnQueryFinished(TSharedPtr<FEnvQueryResult> Result);
+
+	void UseClaimedSmartObject(UBehaviorTreeComponent& OwnerComp, FSmartObjectClaimHandle ClaimHandle, FBTUseSOTaskMemory& MyMemory);
+
 protected:
 	/** Additional tag query to filter available smart objects. We'll query for smart
 	 *	objects that support activities tagged in a way matching the filter.
@@ -39,5 +48,11 @@ protected:
 	FGameplayTagQuery ActivityRequirements;
 
 	UPROPERTY(EditAnywhere, Category = SmartObjects)
+	FEQSParametrizedQueryExecutionRequest EQSRequest;
+
+	/** Used for smart object querying if EQSRequest is not configured */
+	UPROPERTY(EditAnywhere, Category = SmartObjects, meta=(DisplayName="Fallback Radius"))
 	float Radius;
+
+	FQueryFinishedSignature EQSQueryFinishedDelegate; 
 };
