Changed when BTTask's EQSRequestID is cleared to make sure we don't erronously report it not being reset if the "use SO task" fails instantly

diff --git a/Engine/Plugins/Runtime/GameplayBehaviorSmartObjects/Source/GameplayBehaviorSmartObjectsModule/Private/AI/BTTask_FindAndUseGameplayBehaviorSmartObject.cpp b/Engine/Plugins/Runtime/GameplayBehaviorSmartObjects/Source/GameplayBehaviorSmartObjectsModule/Private/AI/BTTask_FindAndUseGameplayBehaviorSmartObject.cpp
--- a/Engine/Plugins/Runtime/GameplayBehaviorSmartObjects/Source/GameplayBehaviorSmartObjectsModule/Private/AI/BTTask_FindAndUseGameplayBehaviorSmartObject.cpp
+++ b/Engine/Plugins/Runtime/GameplayBehaviorSmartObjects/Source/GameplayBehaviorSmartObjectsModule/Private/AI/BTTask_FindAndUseGameplayBehaviorSmartObject.cpp
@@ -194,6 +194,11 @@ void UBTTask_FindAndUseGameplayBehaviorSmartObject::OnQueryFinished(TSharedPtr<F
 		return;
 	}
 
+	// at this point we've already confirmed that QueryResult does indeed corresponds to the the query we're waiting for
+	// so we need to clear the EQSRequestID here in case the next task we're about to issue (the UAITask_UseGameplayBehaviorSmartObject)
+	// might fail instantly and we do check EQSRequestID on tasks end to make sure everything has been cleaned up properly.
+	MyMemory->EQSRequestID = INDEX_NONE;
+
 	bool bSmartObjectClaimed = false;
 
 	if (QueryResult.IsSuccessful() && (QueryResult.Items.Num() >= 1))
@@ -220,8 +225,6 @@ void UBTTask_FindAndUseGameplayBehaviorSmartObject::OnQueryFinished(TSharedPtr<F
 			}
 		}
 	}
-
-	MyMemory->EQSRequestID = INDEX_NONE;
 	
 	if (bSmartObjectClaimed == false)
 	{
