Fix CIS Issue 257499

diff --git a/Engine/Plugins/Runtime/GameplayBehaviorSmartObjects/Source/GameplayBehaviorSmartObjectsModule/Private/AI/BTTask_FindAndUseGameplayBehaviorSmartObject.cpp b/Engine/Plugins/Runtime/GameplayBehaviorSmartObjects/Source/GameplayBehaviorSmartObjectsModule/Private/AI/BTTask_FindAndUseGameplayBehaviorSmartObject.cpp
--- a/Engine/Plugins/Runtime/GameplayBehaviorSmartObjects/Source/GameplayBehaviorSmartObjectsModule/Private/AI/BTTask_FindAndUseGameplayBehaviorSmartObject.cpp
+++ b/Engine/Plugins/Runtime/GameplayBehaviorSmartObjects/Source/GameplayBehaviorSmartObjectsModule/Private/AI/BTTask_FindAndUseGameplayBehaviorSmartObject.cpp
@@ -212,7 +212,7 @@ void UBTTask_FindAndUseGameplayBehaviorSmartObject::OnQueryFinished(TSharedPtr<F
 		}
 		else if (USmartObjectSubsystem* SmartObjectSubsystem = USmartObjectSubsystem::GetCurrent(MyOwner->GetWorld()))
 		{
-			const FSmartObjectActorUserData ActorUserData(Result->Owner.Get());
+			const FSmartObjectActorUserData ActorUserData(Cast<AActor>(Result->Owner.Get());
 			const FConstStructView ActorUserDataView(FConstStructView::Make(ActorUserData));
 
 			// we could use QueryResult.GetItemAsTypeChecked, but the below implementation is more efficient
