UBTTask_FindAndUseGameplayBehaviorSmartObject should reach slot location before using object

#virtualized

diff --git a/Engine/Plugins/Runtime/GameplayBehaviorSmartObjects/Source/GameplayBehaviorSmartObjectsModule/Private/AI/BTTask_FindAndUseGameplayBehaviorSmartObject.cpp b/Engine/Plugins/Runtime/GameplayBehaviorSmartObjects/Source/GameplayBehaviorSmartObjectsModule/Private/AI/BTTask_FindAndUseGameplayBehaviorSmartObject.cpp
--- a/Engine/Plugins/Runtime/GameplayBehaviorSmartObjects/Source/GameplayBehaviorSmartObjectsModule/Private/AI/BTTask_FindAndUseGameplayBehaviorSmartObject.cpp
+++ b/Engine/Plugins/Runtime/GameplayBehaviorSmartObjects/Source/GameplayBehaviorSmartObjectsModule/Private/AI/BTTask_FindAndUseGameplayBehaviorSmartObject.cpp
@@ -243,6 +243,7 @@ void UBTTask_FindAndUseGameplayBehaviorSmartObject::UseClaimedSmartObject(UBehav
 	checkSlow(ClaimHandle.IsValid());
 	UAITask_UseGameplayBehaviorSmartObject* UseSOTask = NewBTAITask<UAITask_UseGameplayBehaviorSmartObject>(OwnerComp);
 	UseSOTask->SetClaimHandle(ClaimHandle);
+	UseSOTask->SetShouldReachSlotLocation(true);
 	UseSOTask->ReadyForActivation();
 
 	MyMemory.TaskInstance = UseSOTask;
