Fixing UAITask_UseGameplayBehaviorSmartObject::Abort to not crash if we have a null pawn.

-When the world tears down, it forces all smart object components to end play. This cascades into smart objects aborting, which leads to the UAITask_UseGameplayBehaviorSmartObject aborting, which hard asserts if the owning controller's pawn is not valid. This fix checks to make sure the pawn is null before dereferencing it and removes the check against the pawn.

diff --git a/Engine/Plugins/Runtime/GameplayBehaviorSmartObjects/Source/GameplayBehaviorSmartObjectsModule/Private/AI/AITask_UseGameplayBehaviorSmartObject.cpp b/Engine/Plugins/Runtime/GameplayBehaviorSmartObjects/Source/GameplayBehaviorSmartObjectsModule/Private/AI/AITask_UseGameplayBehaviorSmartObject.cpp
--- a/Engine/Plugins/Runtime/GameplayBehaviorSmartObjects/Source/GameplayBehaviorSmartObjectsModule/Private/AI/AITask_UseGameplayBehaviorSmartObject.cpp
+++ b/Engine/Plugins/Runtime/GameplayBehaviorSmartObjects/Source/GameplayBehaviorSmartObjectsModule/Private/AI/AITask_UseGameplayBehaviorSmartObject.cpp
@@ -310,9 +310,13 @@ void UAITask_UseGameplayBehaviorSmartObject::Abort()
 		if (GameplayBehavior != nullptr)
 		{
 			check(OwnerController);
-			check(OwnerController->GetPawn());
 			GameplayBehavior->GetOnBehaviorFinishedDelegate().Remove(OnBehaviorFinishedNotifyHandle);
-			GameplayBehavior->AbortBehavior(*OwnerController->GetPawn());
+
+			// Check for a valid pawn first - this is used by the 'shutdown' path so pawn is not guaranteed to still be valid.
+			if (APawn* OwnerPawn = OwnerController->GetPawn())
+			{
+				GameplayBehavior->AbortBehavior(*OwnerPawn);
+			}
 		}
 	}
 
