[SmartObject] fixed check firing in normal use case (abort)

diff --git a/Engine/Plugins/Runtime/GameplayBehaviorSmartObjects/Source/GameplayBehaviorSmartObjectsModule/Private/AI/BTTask_FindAndUseGameplayBehaviorSmartObject.cpp b/Engine/Plugins/Runtime/GameplayBehaviorSmartObjects/Source/GameplayBehaviorSmartObjectsModule/Private/AI/BTTask_FindAndUseGameplayBehaviorSmartObject.cpp
--- a/Engine/Plugins/Runtime/GameplayBehaviorSmartObjects/Source/GameplayBehaviorSmartObjectsModule/Private/AI/BTTask_FindAndUseGameplayBehaviorSmartObject.cpp
+++ b/Engine/Plugins/Runtime/GameplayBehaviorSmartObjects/Source/GameplayBehaviorSmartObjectsModule/Private/AI/BTTask_FindAndUseGameplayBehaviorSmartObject.cpp
@@ -121,14 +121,18 @@ EBTNodeResult::Type UBTTask_FindAndUseGameplayBehaviorSmartObject::AbortTask(UBe
 
 	if (MyMemory->EQSRequestID != INDEX_NONE)
 	{
+		// Clear before calling AbortQuery since it will call us back through OnTaskFinished
+		// which uses a check to make sure that the request ID was cleared before reaching that point.
+		const int32 EQSRequestID = MyMemory->EQSRequestID;
+		MyMemory->EQSRequestID = INDEX_NONE;
+
 		if (UWorld* World = OwnerComp.GetWorld())
 		{
 			if (UEnvQueryManager* EnvQueryManager = UEnvQueryManager::GetCurrent(World))
 			{
-				EnvQueryManager->AbortQuery(MyMemory->EQSRequestID);
+				EnvQueryManager->AbortQuery(EQSRequestID);
 			}
 		}
-		MyMemory->EQSRequestID = INDEX_NONE;
 	}
 
 	return EBTNodeResult::Aborted;
