PR #13135: Fix Smart Objects Crash

diff --git a/Engine/Plugins/Runtime/GameplayBehaviorSmartObjects/Source/GameplayBehaviorSmartObjectsModule/Private/AI/AITask_UseGameplayBehaviorSmartObject.cpp b/Engine/Plugins/Runtime/GameplayBehaviorSmartObjects/Source/GameplayBehaviorSmartObjectsModule/Private/AI/AITask_UseGameplayBehaviorSmartObject.cpp
--- a/Engine/Plugins/Runtime/GameplayBehaviorSmartObjects/Source/GameplayBehaviorSmartObjectsModule/Private/AI/AITask_UseGameplayBehaviorSmartObject.cpp
+++ b/Engine/Plugins/Runtime/GameplayBehaviorSmartObjects/Source/GameplayBehaviorSmartObjectsModule/Private/AI/AITask_UseGameplayBehaviorSmartObject.cpp
@@ -268,10 +268,15 @@ void UAITask_UseGameplayBehaviorSmartObject::OnDestroy(const bool bInOwnerFinish
 {
 	if (ClaimedHandle.IsValid())
 	{
-		USmartObjectSubsystem* SmartObjectSubsystem = USmartObjectSubsystem::GetCurrent(OwnerController->GetWorld());
-		check(SmartObjectSubsystem);
-		SmartObjectSubsystem->MarkSlotAsFree(ClaimedHandle);
-		SmartObjectSubsystem->UnregisterSlotInvalidationCallback(ClaimedHandle);
+		if (ensureMsgf(IsValid(OwnerController), TEXT("OwnerController is expected to always be valid")))
+		{
+			USmartObjectSubsystem* SmartObjectSubsystem = USmartObjectSubsystem::GetCurrent(OwnerController->GetWorld());
+			if (IsValid(SmartObjectSubsystem))
+			{
+				SmartObjectSubsystem->MarkSlotAsFree(ClaimedHandle);
+				SmartObjectSubsystem->UnregisterSlotInvalidationCallback(ClaimedHandle);
+			}
+		}
 		ClaimedHandle.Invalidate();
 	}
 
@@ -325,8 +330,7 @@ void UAITask_UseGameplayBehaviorSmartObject::Abort()
 
 void UAITask_UseGameplayBehaviorSmartObject::OnSmartObjectBehaviorFinished(UGameplayBehavior& Behavior, AActor& Avatar, const bool bInterrupted)
 {
-	// Adding an ensure in case the assumptions change in the future.
-	ensure(OwnerController != nullptr);
+	ensureMsgf(IsValid(OwnerController), TEXT("OwnerController is expected to always be valid"));
 
 	// make sure we handle the right pawn - we can get this notify for a different
 	// Avatar if the behavior sending it out is not instanced (CDO is being used to perform actions)
diff --git a/Engine/Source/Runtime/AIModule/Private/Tasks/AITask_MoveTo.cpp b/Engine/Source/Runtime/AIModule/Private/Tasks/AITask_MoveTo.cpp
--- a/Engine/Source/Runtime/AIModule/Private/Tasks/AITask_MoveTo.cpp
+++ b/Engine/Source/Runtime/AIModule/Private/Tasks/AITask_MoveTo.cpp
@@ -252,7 +252,7 @@ void UAITask_MoveTo::ResetObservers()
 
 void UAITask_MoveTo::ResetTimers()
 {
-	if (OwnerController)
+	if (IsValid(OwnerController) && IsValid(OwnerController->GetWorld()))
 	{
 		// Remove all timers including the ones that might have been set with SetTimerForNextTick 
 		OwnerController->GetWorldTimerManager().ClearAllTimersForObject(this);
